Bugfix for incorrect generation of RSA keys

diff --git a/Engine/Plugins/Editor/CryptoKeys/Source/CryptoKeysOpenSSL/Private/CryptoKeysOpenSSL.cpp b/Engine/Plugins/Editor/CryptoKeys/Source/CryptoKeysOpenSSL/Private/CryptoKeysOpenSSL.cpp
--- a/Engine/Plugins/Editor/CryptoKeys/Source/CryptoKeysOpenSSL/Private/CryptoKeysOpenSSL.cpp
+++ b/Engine/Plugins/Editor/CryptoKeys/Source/CryptoKeysOpenSSL/Private/CryptoKeysOpenSSL.cpp
@@ -33,10 +33,10 @@ namespace CryptoKeysOpenSSL
 
 	void BigNumToArray(const BIGNUM* InNum, TArray<uint8>& OutBytes, int32 InKeySize)
 	{
-		int NumBytes = BN_num_bytes(InNum);
+		int32 NumBytes = BN_num_bytes(InNum);
 		check(NumBytes <= InKeySize);
 		OutBytes.SetNumZeroed(InKeySize);
-		BN_bn2bin(InNum, OutBytes.GetData());
+		BN_bn2bin(InNum, OutBytes.GetData() + (InKeySize - NumBytes));
 		Algo::Reverse(OutBytes); // bn2bin produces big endian data
 	}
 
