Renamed "Crypto" project settings section to "Encryption" to better describe it
#rnx

diff --git a/Engine/Plugins/Editor/CryptoKeys/Source/CryptoKeys/Private/CryptoKeys.cpp b/Engine/Plugins/Editor/CryptoKeys/Source/CryptoKeys/Private/CryptoKeys.cpp
--- a/Engine/Plugins/Editor/CryptoKeys/Source/CryptoKeys/Private/CryptoKeys.cpp
+++ b/Engine/Plugins/Editor/CryptoKeys/Source/CryptoKeys/Private/CryptoKeys.cpp
@@ -44,9 +44,9 @@ public:
 	{
 		if (ISettingsModule* SettingsModule = FModuleManager::GetModulePtr<ISettingsModule>("Settings"))
 		{
-			SettingsModule->RegisterSettings("Project", "Project", "Crypto",
-				LOCTEXT("CryptoSettingsName", "Crypto"),
-				LOCTEXT("CryptoSettingsDescription", "Configure the project crypto keys"),
+			SettingsModule->RegisterSettings("Project", "Project", "Encryption",
+				LOCTEXT("CryptoSettingsName", "Encryption"),
+				LOCTEXT("CryptoSettingsDescription", "Configure the project encryption keys"),
 				GetMutableDefault<UCryptoKeysSettings>());
 		}
 	}
@@ -55,7 +55,7 @@ public:
 	{
 		if (ISettingsModule* SettingsModule = FModuleManager::GetModulePtr<ISettingsModule>("Settings"))
 		{
-			SettingsModule->UnregisterSettings("Project", "Project", "Crypto");
+			SettingsModule->UnregisterSettings("Project", "Project", "Encryption");
 		}
 	}
 
