- Fixed an issue with Clearing Encryption Key wasn't saving, as it was using a hacky way to save to defaults (unlike Generate Encryption Key which was using the correct method)

diff --git a/Engine/Plugins/Editor/CryptoKeys/Source/CryptoKeys/Private/CryptoKeysSettingsDetails.cpp b/Engine/Plugins/Editor/CryptoKeys/Source/CryptoKeys/Private/CryptoKeysSettingsDetails.cpp
--- a/Engine/Plugins/Editor/CryptoKeys/Source/CryptoKeys/Private/CryptoKeysSettingsDetails.cpp
+++ b/Engine/Plugins/Editor/CryptoKeys/Source/CryptoKeys/Private/CryptoKeysSettingsDetails.cpp
@@ -59,7 +59,7 @@ void FCryptoKeysSettingsDetails::CustomizeDetails(IDetailLayoutBuilder& DetailBu
 				.OnClicked_Lambda([this, Settings]()
 				{
 					Settings->EncryptionKey = TEXT("");
-					Settings->SaveConfig(CPF_Config, *Settings->GetDefaultConfigFilename());
+					Settings->TryUpdateDefaultConfigFile();
 					return(FReply::Handled());
 				})
 			]
